Added csv parser

The csv parser takes a string buffer which it splits into a series of null terminated, unescaped cell strings. In this way it can completely avoid heap allocation (esecially if constructed with a moved string). The parser copes with quoted cell values, and attempts to behave in the same way as Microsoft Excel for interoperability.

diff --git a/Engine/Source/Developer/CsvParser/CsvParser.Build.cs b/Engine/Source/Developer/CsvParser/CsvParser.Build.cs
new file mode 100644
--- /dev/null
+++ b/Engine/Source/Developer/CsvParser/CsvParser.Build.cs
@@ -0,0 +1,16 @@
+﻿// Copyright 1998-2014 Epic Games, Inc. All Rights Reserved.
+
+using UnrealBuildTool;
+
+public class CsvParser : ModuleRules
+{
+    public CsvParser(TargetInfo Target)
+	{
+		PrivateDependencyModuleNames.AddRange(
+			new string[] 
+			{ 
+				"Core"
+			}
+            );
+	}
+}
diff --git a/Engine/Source/Developer/CsvParser/Private/CsvParser.cpp b/Engine/Source/Developer/CsvParser/Private/CsvParser.cpp
new file mode 100644
--- /dev/null
+++ b/Engine/Source/Developer/CsvParser/Private/CsvParser.cpp
@@ -0,0 +1,141 @@
+// Copyright 1998-2014 Epic Games, Inc. All Rights Reserved.
+
+#include "CsvParserPrivatePCH.h"
+#include "CsvParser.h"
+
+FCsvParser::FCsvParser(FString InSourceString)
+	: SourceString(MoveTemp(InSourceString))		
+{
+	if (!SourceString.IsEmpty())
+	{
+		ParseRows();
+	}
+}
+		
+void FCsvParser::ParseRows()
+{
+	BufferStart = &SourceString[0];
+	ReadAt = BufferStart;
+
+	EParseResult Result;
+	do
+	{
+		Result = ParseRow();
+	} while(Result != EParseResult::EndOfString);
+}
+
+FCsvParser::EParseResult FCsvParser::ParseRow()
+{
+	// Check for an empty line
+	const int8 NewLineSize = MeasureNewLine(ReadAt);
+	if (NewLineSize)
+	{
+		ReadAt += NewLineSize;
+		return *ReadAt ? EParseResult::EndOfCell : EParseResult::EndOfString;
+	}
+
+	EParseResult Result;
+
+	Rows.Emplace();
+	do
+	{
+		Result = ParseCell();
+	}
+	while(Result == EParseResult::EndOfCell);
+
+	return Result;
+}
+
+FCsvParser::EParseResult FCsvParser::ParseCell()
+{
+	TCHAR* WriteAt = BufferStart + (ReadAt - BufferStart);
+
+	// Check if this cell is quoted. Whitespace between cell opening and quote is invalid.
+	bool bQuoted = *ReadAt == '"';
+
+	if (bQuoted)
+	{
+		// Skip over the first quote
+		ReadAt = ++WriteAt;
+	}
+	
+	Rows.Last().Add(ReadAt);
+
+	while( *ReadAt )
+	{
+		// If the cell value is quoted, we swallow everything until we find a closing quote
+		if (bQuoted)
+		{
+			if (*ReadAt == '"')
+			{
+				// RFC 4180 specifies that double quotes are escaped as ""
+
+				int32 NumQuotes = 0;
+				while (*(ReadAt + NumQuotes) == '"') ++NumQuotes;
+
+				if (NumQuotes % 2 != 0)
+				{
+					// No longer quoted if there are an even number of quotes before this one
+					bQuoted = false;
+				}
+
+				ReadAt += NumQuotes;
+
+				// Unescape the double quotes
+				// We leave the write pos pointing at the trailing closing quote if present so
+				// it gets overwritten by any subsequent text in the cell
+				NumQuotes /= 2;
+				while(NumQuotes-- > 0) *(WriteAt++) = '"';
+
+				continue;
+			}
+		}
+		else
+		{
+			// Check for the end of a row (new line)
+			const int8 NewLineSize = MeasureNewLine(ReadAt);
+			if (NewLineSize != 0)
+			{
+				// Null terminate the cell
+				*WriteAt = '\0';
+				ReadAt += NewLineSize;
+
+				return *ReadAt ? EParseResult::EndOfRow : EParseResult::EndOfString;
+			}
+			else if (*ReadAt == ',')
+			{
+				*WriteAt = '\0';
+				++ReadAt;
+
+				return *ReadAt ? EParseResult::EndOfCell : EParseResult::EndOfString;
+			}
+		}
+		if (WriteAt != ReadAt)
+		{
+			(*WriteAt++) = (*ReadAt++);
+		}
+		else
+		{
+			ReadAt = ++WriteAt;
+		}
+	}
+
+	return EParseResult::EndOfString;
+}
+
+int8 FCsvParser::MeasureNewLine(const TCHAR* At)
+{
+	switch(*At)
+	{
+		case '\r':
+			if (*(At+1) == '\n')	// could be null
+			{
+				return 2;
+			}
+			return 1;
+		case '\n':
+			return 1;
+		default:
+			return 0;
+	}
+}
diff --git a/Engine/Source/Developer/CsvParser/Private/CsvParserPrivatePCH.h b/Engine/Source/Developer/CsvParser/Private/CsvParserPrivatePCH.h
new file mode 100644
--- /dev/null
+++ b/Engine/Source/Developer/CsvParser/Private/CsvParserPrivatePCH.h
@@ -0,0 +1,5 @@
+// Copyright 1998-2014 Epic Games, Inc. All Rights Reserved.
+
+#pragma once
+
+#include "Core.h"
\ No newline at end of file
diff --git a/Engine/Source/Developer/CsvParser/Private/CsvParserTests.cpp b/Engine/Source/Developer/CsvParser/Private/CsvParserTests.cpp
new file mode 100644
--- /dev/null
+++ b/Engine/Source/Developer/CsvParser/Private/CsvParserTests.cpp
@@ -0,0 +1,129 @@
+// Copyright 1998-2014 Epic Games, Inc. All Rights Reserved.
+
+#include "CsvParserPrivatePCH.h"
+#include "CsvParser.h"
+
+namespace CsvParser_Tests
+{
+
+	/** Check that the specified row of cells matches what we expect */
+	template <uint32 N>
+	bool CheckRow(const TArray<const TCHAR*>& Row, const FString (&Expected)[N], int32 RowIndex, FAutomationTestBase* Test)
+	{
+		if (Row.Num() != N)
+		{
+			Test->AddError(FString::Printf(TEXT("Row %d has the wrong number of cells. Expected %d, but parsed result has %d."), RowIndex, N, Row.Num()));
+			return false;
+		}
+
+		for (int32 Cell = 0; Cell < Row.Num(); ++Cell)
+		{
+			if (FCString::Strcmp(Row[Cell], *Expected[Cell]) != 0)
+			{
+				Test->AddError(FString::Printf(TEXT("Cell %d on row %d did not parse correctly. Expected:\n%s\nBut parsed result was:\n%s"), Cell, RowIndex, *Expected[Cell], Row[Cell]));
+				return false;
+			}
+		}
+
+		return true;
+	}
+	/** Check that the specified rows all match the specified results */
+	template <uint32 RowCount, uint32 CellCount>
+	bool CheckRows(const FCsvParser::FRows& Rows, const FString (&Expected)[RowCount][CellCount], FAutomationTestBase* Test)
+	{
+		if (Rows.Num() != RowCount)
+		{
+			Test->AddError(FString::Printf(TEXT("Parser generated wrong number of rows. Expected %d, but parsed result has %d."), RowCount, Rows.Num()));
+			return false;
+		}
+
+		for (int32 Row = 0; Row < RowCount; ++Row)
+		{
+			if (!CheckRow(Rows[Row], Expected[Row], Row, Test))
+			{
+				return false;
+			}
+		}
+
+		return true;
+	}
+}
+
+IMPLEMENT_SIMPLE_AUTOMATION_TEST(FBasicTest, "Engine.CSV Parser.Basic", 0)
+
+bool FBasicTest::RunTest(const FString& Parameters)
+{
+	FString Str =	"1,2,3,4\n5,6,7,8";
+	const FCsvParser Parser(MoveTemp(Str));
+
+	const FString Expected[][4] = {
+		{ "1", "2", "3", "4" },
+		{ "5", "6", "7", "8" }
+	};
+
+	return CsvParser_Tests::CheckRows(Parser.GetRows(), Expected, this);
+}
+
+IMPLEMENT_SIMPLE_AUTOMATION_TEST(FQuotedStringTest, "Engine.CSV Parser.Quoted String", 0)
+
+bool FQuotedStringTest::RunTest(const FString& Parameters)
+{
+	FString Str =	"\"Quoted with nested \"\"quotes\"\", and \"\"commas\"\"\" \"unquoted due to \"whitespace\" , second \" unquoted\n"
+					"\"Quoted\nString,With\nNewlines\",\"\"\n"
+					"1,\"2\"\n";
+	const FCsvParser Parser(MoveTemp(Str));
+
+	const FString Expected[][2] = {
+		{ "Quoted with nested \"quotes\", and \"commas\" \"unquoted due to \"whitespace\" ", " second \" unquoted" },
+		{ "Quoted\nString,With\nNewlines", "" },
+		{ "1", "2" },
+	};
+
+	return CsvParser_Tests::CheckRows(Parser.GetRows(), Expected, this);
+}
+
+IMPLEMENT_SIMPLE_AUTOMATION_TEST(FEmptyCellsTest, "Engine.CSV Parser.Empty Cells", 0)
+
+bool FEmptyCellsTest::RunTest(const FString& Parameters)
+{
+	FString Str =	",,,\n"
+					",,,\n"
+					" , , , \n";
+	const FCsvParser Parser(MoveTemp(Str));
+
+	const FString Expected[][4] = {
+		{ "", "", "", "" },
+		{ "", "", "", "" },
+		{ " ", " ", " ", " " },	// Whitespace is significant in csv
+	};
+
+	return CsvParser_Tests::CheckRows(Parser.GetRows(), Expected, this);
+}
+
+IMPLEMENT_SIMPLE_AUTOMATION_TEST(FMixedNewLinesTest, "Engine.CSV Parser.Mixed Newlines", 0)
+
+bool FMixedNewLinesTest::RunTest(const FString& Parameters)
+{
+	const FCsvParser Parser(FString("1\r2\n3\r\n4\n\n5\r\r6\r\n\r\n7\n\r8"));
+
+	const FString Expected[][1] = {
+		{ "1" }, { "2" }, { "3" }, { "4" }, { "5" }, { "6" }, { "7" }, { "8" },
+	};
+
+	return CsvParser_Tests::CheckRows(Parser.GetRows(), Expected, this);
+}
+
+IMPLEMENT_SIMPLE_AUTOMATION_TEST(FEmptyStringTest, "Engine.CSV Parser.Empty String", 0)
+
+bool FEmptyStringTest::RunTest(const FString& Parameters)
+{
+	return	FCsvParser(FString()).GetRows().Num() == 0 &&
+			FCsvParser(FString("")).GetRows().Num() == 0;
+}
+
+IMPLEMENT_SIMPLE_AUTOMATION_TEST(FOnlyNewLinesTest, "Engine.CSV Parser.Only Newlines", 0)
+
+bool FOnlyNewLinesTest::RunTest(const FString& Parameters)
+{
+	return	FCsvParser(FString("\r\n\r\r\n\n\r\n\n\r")).GetRows().Num() == 0;
+}
\ No newline at end of file
diff --git a/Engine/Source/Developer/CsvParser/Public/CsvParser.h b/Engine/Source/Developer/CsvParser/Public/CsvParser.h
new file mode 100644
--- /dev/null
+++ b/Engine/Source/Developer/CsvParser/Public/CsvParser.h
@@ -0,0 +1,63 @@
+// Copyright 1998-2014 Epic Games, Inc. All Rights Reserved.
+
+#pragma once
+
+
+/** A simple, efficient csv parser */
+struct FCsvParser
+{
+	typedef TArray<TArray<const TCHAR*>> FRows;
+
+	/**
+	 *	Construct with a string.
+	 *	Takes a copy because the parser tramples over the file to create null-terminated cell strings
+	 *	Avoid a copy by passing overship of the source string with MoveTemp()
+	 */
+	CSVPARSER_API FCsvParser(FString InSourceString);
+
+	/**
+	 *	Const access to the parsed rows
+	 *	Rows are stored as arrays of parsed, null-terminated cell-strings.
+	 *	Object (and contained strings) is valid for the lifetime of the parser.
+	 */
+	const FRows& GetRows() const
+	{
+		return Rows;
+	}
+
+private:
+	/** Non copyable - has pointers to itself */
+	FCsvParser(const FCsvParser&);
+	FCsvParser& operator=(const FCsvParser&);
+
+	/** Internal enum to control parsing progress */
+	enum struct EParseResult { EndOfCell, EndOfRow, EndOfString };
+
+	/** Parse all the rows in the source data */
+	void ParseRows();
+
+	/** Parse a single row at the current read position */
+	EParseResult ParseRow();
+
+	/** Parse a cell at the current read position */
+	EParseResult ParseCell();
+
+private:
+
+	/** Raw csv buffer, sliced and diced by the parser */
+	FString SourceString;
+
+	/** Fixed pointer to the start of the buffer */
+	TCHAR* BufferStart;
+
+	/** Current parser read position */
+	const TCHAR* ReadAt;
+
+	/** Array of cell string arrays */
+	FRows Rows;
+
+private:
+
+	/** Helper to measure the size of a new line at the current read position */
+	int8 MeasureNewLine(const TCHAR* At);
+};
