- Fixed CsvParser failing to link in monolithic builds

diff --git a/Engine/Source/Developer/CsvParser/Private/CsvParser.cpp b/Engine/Source/Developer/CsvParser/Private/CsvParser.cpp
--- a/Engine/Source/Developer/CsvParser/Private/CsvParser.cpp
+++ b/Engine/Source/Developer/CsvParser/Private/CsvParser.cpp
@@ -2,6 +2,7 @@
 
 #include "CsvParserPrivatePCH.h"
 #include "CsvParser.h"
+#include "ModuleManager.h"
 
 FCsvParser::FCsvParser(FString InSourceString)
 	: SourceString(MoveTemp(InSourceString))		
@@ -139,3 +140,9 @@ int8 FCsvParser::MeasureNewLine(const TCHAR* At)
 			return 0;
 	}
 }
+
+class FCsvParserModule : public IModuleInterface
+{
+};
+
+IMPLEMENT_MODULE(FCsvParserModule, CsvParser)
